UE-127432 //UE5/Main - Compile UnrealGame HoloLens - Detected paths more than 200 characters below UE root directory
-Duplicating cl 17613973 where the shortname is set to fix the path length.

diff --git a/Engine/Plugins/Experimental/AbilitySystemGameFeatureActions/Source/AbilitySystemGameFeatureActions/AbilitySystemGameFeatureActions.Build.cs b/Engine/Plugins/Experimental/AbilitySystemGameFeatureActions/Source/AbilitySystemGameFeatureActions/AbilitySystemGameFeatureActions.Build.cs
--- a/Engine/Plugins/Experimental/AbilitySystemGameFeatureActions/Source/AbilitySystemGameFeatureActions/AbilitySystemGameFeatureActions.Build.cs
+++ b/Engine/Plugins/Experimental/AbilitySystemGameFeatureActions/Source/AbilitySystemGameFeatureActions/AbilitySystemGameFeatureActions.Build.cs
@@ -6,6 +6,8 @@ namespace UnrealBuildTool.Rules
 	{
         public AbilitySystemGameFeatureActions(ReadOnlyTargetRules Target) : base(Target)
 		{
+			ShortName = "AbilitySysGFAs";
+
 			PublicDependencyModuleNames.AddRange(
 				new string[]
                 {
