#UE Do not load attribute tables if you do not have them. This can happen if folks opt to not sync certain folders

diff --git a/Engine/Plugins/Experimental/AbilitySystemGameFeatureActions/Source/AbilitySystemGameFeatureActions/Private/GameFeatureAction_AddAttributeDefaults.cpp b/Engine/Plugins/Experimental/AbilitySystemGameFeatureActions/Source/AbilitySystemGameFeatureActions/Private/GameFeatureAction_AddAttributeDefaults.cpp
--- a/Engine/Plugins/Experimental/AbilitySystemGameFeatureActions/Source/AbilitySystemGameFeatureActions/Private/GameFeatureAction_AddAttributeDefaults.cpp
+++ b/Engine/Plugins/Experimental/AbilitySystemGameFeatureActions/Source/AbilitySystemGameFeatureActions/Private/GameFeatureAction_AddAttributeDefaults.cpp
@@ -2,6 +2,7 @@
 
 #include "GameFeatureAction_AddAttributeDefaults.h"
 #include "AbilitySystemGlobals.h"
+#include "Misc/PackageName.h"
 
 #include UE_INLINE_GENERATED_CPP_BY_NAME(GameFeatureAction_AddAttributeDefaults)
 
@@ -12,8 +13,23 @@
 
 void UGameFeatureAction_AddAttributeDefaults::OnGameFeatureRegistering()
 {
+	const TArray<FSoftObjectPath>* AttribDefaultTableNamesToAdd = &AttribDefaultTableNames;
+
+#if WITH_EDITOR
+	// Do a file exists check in editor builds since some folks do not sync all data in the editor. Ideally we don't need to load anything at GFD registration time, but for now we will do this.
+	TArray<FSoftObjectPath> AttribDefaultTableNamesThatExist;
+	AttribDefaultTableNamesToAdd = &AttribDefaultTableNamesThatExist;
+	for (const FSoftObjectPath& Path : AttribDefaultTableNames)
+	{
+		if (FPackageName::DoesPackageExist(Path.GetLongPackageName()))
+		{
+			AttribDefaultTableNamesThatExist.Add(Path);
+		}
+	}
+#endif // WITH_EDITOR
+
 	UAbilitySystemGlobals& AbilitySystemGlobals = UAbilitySystemGlobals::Get();
-	AbilitySystemGlobals.AddAttributeDefaultTables(AttribDefaultTableNames);
+	AbilitySystemGlobals.AddAttributeDefaultTables(*AttribDefaultTableNamesToAdd);
 }
 
 //////////////////////////////////////////////////////////////////////
