Added toggle to apply attribute set defaults in OnGameFeatureActivating

diff --git a/Engine/Plugins/Experimental/AbilitySystemGameFeatureActions/Source/AbilitySystemGameFeatureActions/Private/GameFeatureAction_AddAttributeDefaults.cpp b/Engine/Plugins/Experimental/AbilitySystemGameFeatureActions/Source/AbilitySystemGameFeatureActions/Private/GameFeatureAction_AddAttributeDefaults.cpp
--- a/Engine/Plugins/Experimental/AbilitySystemGameFeatureActions/Source/AbilitySystemGameFeatureActions/Private/GameFeatureAction_AddAttributeDefaults.cpp
+++ b/Engine/Plugins/Experimental/AbilitySystemGameFeatureActions/Source/AbilitySystemGameFeatureActions/Private/GameFeatureAction_AddAttributeDefaults.cpp
@@ -18,6 +18,98 @@ namespace GameFeatureAction_AddAttributeDefaults
 // UGameFeatureAction_AddAttributeDefaults
 
 void UGameFeatureAction_AddAttributeDefaults::OnGameFeatureRegistering()
+{
+	Super::OnGameFeatureRegistering();
+
+	if (ShouldAddAttributeDefaults())
+	{
+		AddAttributeDefaults();
+	}
+}
+
+void UGameFeatureAction_AddAttributeDefaults::OnGameFeatureActivating(FGameFeatureActivatingContext& Context)
+{
+	Super::OnGameFeatureActivating(Context);
+
+	if (ShouldAddAttributeDefaults())
+	{
+		AddAttributeDefaults();
+	}
+}
+
+void UGameFeatureAction_AddAttributeDefaults::OnGameFeatureUnregistering()
+{
+	if (ShouldRemoveAttributeDefaults())
+	{
+		RemoveAttributeDefaults();
+	}
+
+	Super::OnGameFeatureUnregistering();
+}
+
+void UGameFeatureAction_AddAttributeDefaults::OnGameFeatureDeactivating(FGameFeatureDeactivatingContext& Context)
+{
+	if (ShouldRemoveAttributeDefaults())
+	{
+		RemoveAttributeDefaults();
+	}
+
+	Super::OnGameFeatureDeactivating(Context);
+}
+
+#if WITH_EDITOR
+void UGameFeatureAction_AddAttributeDefaults::PostEditChangeProperty(FPropertyChangedEvent& PropertyChangedEvent)
+{
+	Super::PostEditChangeProperty(PropertyChangedEvent);
+
+	const FName PropertyName = PropertyChangedEvent.GetPropertyName();
+
+	if (PropertyName == GET_MEMBER_NAME_CHECKED(ThisClass, bApplyOnRegister))
+	{
+		// Re-check whether we should apply our modified defaults.
+		// Avoids 'leaking' changes in the event the GFA is unregistered.
+
+		if (ShouldRemoveAttributeDefaults())
+		{
+			RemoveAttributeDefaults();
+		}
+
+		if (ShouldAddAttributeDefaults())
+		{
+			AddAttributeDefaults();
+		}
+	}
+}
+#endif // WITH_EDITOR
+
+bool UGameFeatureAction_AddAttributeDefaults::ShouldAddAttributeDefaults() const
+{
+	// Necessary as during OnGameFeatureRegistering/Activating the plugin is *Registering* but not *Registered*
+	constexpr bool bCheckForRegisteringAndActivating = true;
+
+	if (!bAttributesHaveBeenSet)
+	{
+		return bApplyOnRegister
+			? IsGameFeaturePluginRegistered(bCheckForRegisteringAndActivating)
+			: IsGameFeaturePluginActive(bCheckForRegisteringAndActivating);
+	}
+
+	return false;
+}
+
+bool UGameFeatureAction_AddAttributeDefaults::ShouldRemoveAttributeDefaults() const
+{
+	if (bAttributesHaveBeenSet)
+	{
+		return bApplyOnRegister
+			? !IsGameFeaturePluginRegistered()
+			: !IsGameFeaturePluginActive();
+	}
+
+	return false;
+}
+
+void UGameFeatureAction_AddAttributeDefaults::AddAttributeDefaults()
 {
 	const TArray<FSoftObjectPath>* AttribDefaultTableNamesToAdd = &AttribDefaultTableNames;
 
@@ -43,15 +135,19 @@ void UGameFeatureAction_AddAttributeDefaults::OnGameFeatureRegistering()
 		UAbilitySystemGlobals& AbilitySystemGlobals = UAbilitySystemGlobals::Get();
 		AbilitySystemGlobals.AddAttributeDefaultTables(AttributeDefaultTablesOwnerName, *AttribDefaultTableNamesToAdd);
 	}
+
+	bAttributesHaveBeenSet = true;
 }
 
-void UGameFeatureAction_AddAttributeDefaults::OnGameFeatureUnregistering()
+void UGameFeatureAction_AddAttributeDefaults::RemoveAttributeDefaults()
 {
 	if (!AttribDefaultTableNames.IsEmpty() && GameFeatureAction_AddAttributeDefaults::CVarAllowRemoveAttributeDefaultTables.GetValueOnAnyThread())
 	{
 		UAbilitySystemGlobals& AbilitySystemGlobals = UAbilitySystemGlobals::Get();
 		AbilitySystemGlobals.RemoveAttributeDefaultTables(AttributeDefaultTablesOwnerName, AttribDefaultTableNames);
 	}
+
+	bAttributesHaveBeenSet = false;
 }
 
 //////////////////////////////////////////////////////////////////////
diff --git a/Engine/Plugins/Experimental/AbilitySystemGameFeatureActions/Source/AbilitySystemGameFeatureActions/Public/GameFeatureAction_AddAttributeDefaults.h b/Engine/Plugins/Experimental/AbilitySystemGameFeatureActions/Source/AbilitySystemGameFeatureActions/Public/GameFeatureAction_AddAttributeDefaults.h
--- a/Engine/Plugins/Experimental/AbilitySystemGameFeatureActions/Source/AbilitySystemGameFeatureActions/Public/GameFeatureAction_AddAttributeDefaults.h
+++ b/Engine/Plugins/Experimental/AbilitySystemGameFeatureActions/Source/AbilitySystemGameFeatureActions/Public/GameFeatureAction_AddAttributeDefaults.h
@@ -20,15 +20,37 @@ class UGameFeatureAction_AddAttributeDefaults final : public UGameFeatureAction
 public:
 	//~UGameFeatureAction interface
 	virtual void OnGameFeatureRegistering() override;
+	virtual void OnGameFeatureActivating(FGameFeatureActivatingContext& Context) override;
 	virtual void OnGameFeatureUnregistering() override;
+	virtual void OnGameFeatureDeactivating(FGameFeatureDeactivatingContext& Context) override;
 	//~End of UGameFeatureAction interface
 
+#if WITH_EDITOR
+	virtual void PostEditChangeProperty(FPropertyChangedEvent& PropertyChangedEvent) override;
+#endif // WITH_EDITOR
+
+	/**
+	 * True: Apply defaults when the game feature is registered.
+	 * False: Apply defaults when the game feature is activated.
+	 */
+	UPROPERTY(EditAnywhere, Category = Attributes, AdvancedDisplay)
+	bool bApplyOnRegister = true;
+
 	/** List of attribute default tables to add */
 	UPROPERTY(EditAnywhere, Category = Attributes)
 	TArray<FSoftObjectPath> AttribDefaultTableNames;
 
 private:
+
+	bool ShouldAddAttributeDefaults() const;
+	bool ShouldRemoveAttributeDefaults() const;
+
+	void AddAttributeDefaults();
+	void RemoveAttributeDefaults();
+
 	FName AttributeDefaultTablesOwnerName;
+
+	bool bAttributesHaveBeenSet = false;
 };
 
 #if UE_ENABLE_INCLUDE_ORDER_DEPRECATED_IN_5_2
diff --git a/Engine/Plugins/Runtime/GameFeatures/Source/GameFeatures/Private/GameFeatureAction.cpp b/Engine/Plugins/Runtime/GameFeatures/Source/GameFeatures/Private/GameFeatureAction.cpp
--- a/Engine/Plugins/Runtime/GameFeatures/Source/GameFeatures/Private/GameFeatureAction.cpp
+++ b/Engine/Plugins/Runtime/GameFeatures/Source/GameFeatures/Private/GameFeatureAction.cpp
@@ -24,14 +24,14 @@ void UGameFeatureAction::OnGameFeatureActivating(FGameFeatureActivatingContext&
 	OnGameFeatureActivating();
 }
 
-bool UGameFeatureAction::IsGameFeaturePluginRegistered() const
+bool UGameFeatureAction::IsGameFeaturePluginRegistered(bool bCheckForRegistering  /*= false*/) const
 {
 	UGameFeatureData* GameFeatureData = GetGameFeatureData();
-	return !!GameFeatureData ? GameFeatureData->IsGameFeaturePluginRegistered() : false;
+	return !!GameFeatureData ? GameFeatureData->IsGameFeaturePluginRegistered(bCheckForRegistering) : false;
 }
 
-bool UGameFeatureAction::IsGameFeaturePluginActive() const
+bool UGameFeatureAction::IsGameFeaturePluginActive(bool bCheckForActivating /*= false*/) const
 {
 	UGameFeatureData* GameFeatureData = GetGameFeatureData();
-	return !!GameFeatureData ? GameFeatureData->IsGameFeaturePluginActive() : false;
+	return !!GameFeatureData ? GameFeatureData->IsGameFeaturePluginActive(bCheckForActivating) : false;
 }
diff --git a/Engine/Plugins/Runtime/GameFeatures/Source/GameFeatures/Private/GameFeatureData.cpp b/Engine/Plugins/Runtime/GameFeatures/Source/GameFeatures/Private/GameFeatureData.cpp
--- a/Engine/Plugins/Runtime/GameFeatures/Source/GameFeatures/Private/GameFeatureData.cpp
+++ b/Engine/Plugins/Runtime/GameFeatures/Source/GameFeatures/Private/GameFeatureData.cpp
@@ -774,26 +774,26 @@ void UGameFeatureData::GetPluginName(const UGameFeatureData* GFD, FString& Plugi
 	}
 }
 
-bool UGameFeatureData::IsGameFeaturePluginRegistered() const
+bool UGameFeatureData::IsGameFeaturePluginRegistered(bool bCheckForRegistering /*= false*/) const
 {
 	FString PluginURL;
 	FString PluginName;
 	GetPluginName(PluginName);
 	if (UGameFeaturesSubsystem::Get().GetPluginURLByName(PluginName, PluginURL))
 	{
-		return UGameFeaturesSubsystem::Get().IsGameFeaturePluginRegistered(PluginURL);
+		return UGameFeaturesSubsystem::Get().IsGameFeaturePluginRegistered(PluginURL, bCheckForRegistering);
 	}
 	return false;
 }
 
-bool UGameFeatureData::IsGameFeaturePluginActive() const
+bool UGameFeatureData::IsGameFeaturePluginActive(bool bCheckForActivating /*= false*/) const
 {
 	FString PluginURL;
 	FString PluginName;
 	GetPluginName(PluginName);
 	if (UGameFeaturesSubsystem::Get().GetPluginURLByName(PluginName, PluginURL))
 	{
-		return UGameFeaturesSubsystem::Get().IsGameFeaturePluginActive(PluginURL);
+		return UGameFeaturesSubsystem::Get().IsGameFeaturePluginActive(PluginURL, bCheckForActivating);
 	}
 	return false;
 }
diff --git a/Engine/Plugins/Runtime/GameFeatures/Source/GameFeatures/Public/GameFeatureAction.h b/Engine/Plugins/Runtime/GameFeatures/Source/GameFeatures/Public/GameFeatureAction.h
--- a/Engine/Plugins/Runtime/GameFeatures/Source/GameFeatures/Public/GameFeatureAction.h
+++ b/Engine/Plugins/Runtime/GameFeatures/Source/GameFeatures/Public/GameFeatureAction.h
@@ -40,10 +40,10 @@ public:
 	virtual void OnGameFeatureDeactivating(FGameFeatureDeactivatingContext& Context) {}
 
 	/** Returns whether the action game feature plugin is registered or not. */
-	bool IsGameFeaturePluginRegistered() const;
+	bool IsGameFeaturePluginRegistered(bool bCheckForRegistering = false) const;
 
 	/** Returns whether the action game feature plugin is active or not. */
-	bool IsGameFeaturePluginActive() const;
+	bool IsGameFeaturePluginActive(bool bCheckForActivating = false) const;
 
 #if WITH_EDITORONLY_DATA
 	virtual void AddAdditionalAssetBundleData(FAssetBundleData& AssetBundleData) {}
diff --git a/Engine/Plugins/Runtime/GameFeatures/Source/GameFeatures/Public/GameFeatureData.h b/Engine/Plugins/Runtime/GameFeatures/Source/GameFeatures/Public/GameFeatureData.h
--- a/Engine/Plugins/Runtime/GameFeatures/Source/GameFeatures/Public/GameFeatureData.h
+++ b/Engine/Plugins/Runtime/GameFeatures/Source/GameFeatures/Public/GameFeatureData.h
@@ -50,10 +50,10 @@ public:
 	void GetPluginName(FString& PluginName) const;
 
 	/** Returns whether the game feature plugin is registered or not. */
-	bool IsGameFeaturePluginRegistered() const;
+	bool IsGameFeaturePluginRegistered(bool bCheckForRegistering = false) const;
 
 	/** Returns whether the game feature plugin is active or not. */
-	bool IsGameFeaturePluginActive() const;
+	bool IsGameFeaturePluginActive(bool bCheckForActivating = false) const;
 
 	/**
 	 * Returns the install bundle name if one exists for this plugin.
